Slate : Toolbox now in Runtime/ and unversally accessible via Ctrl+Shift+~. Disabled for UE_BUILD_SHIPPING || UE_BUILD_TEST.

diff --git a/Engine/Source/Developer/Toolbox/Private/Toolbox.cpp b/Engine/Source/Developer/Toolbox/Private/Toolbox.cpp
deleted file mode 100644
diff --git a/Engine/Source/Developer/Toolbox/Public/Toolbox.h b/Engine/Source/Developer/Toolbox/Public/Toolbox.h
deleted file mode 100644
diff --git a/Engine/Source/Developer/Toolbox/Toolbox.Build.cs b/Engine/Source/Developer/Toolbox/Toolbox.Build.cs
deleted file mode 100644
diff --git a/Engine/Source/Editor/MainFrame/MainFrame.Build.cs b/Engine/Source/Editor/MainFrame/MainFrame.Build.cs
--- a/Engine/Source/Editor/MainFrame/MainFrame.Build.cs
+++ b/Engine/Source/Editor/MainFrame/MainFrame.Build.cs
@@ -53,6 +53,7 @@ public class MainFrame : ModuleRules
 				"SuperSearch",
 				"Settings",
 				"SourceCodeAccess",
+                "Toolbox",
 			}
 		);
 
diff --git a/Engine/Source/Editor/MainFrame/Private/Menus/MainMenu.cpp b/Engine/Source/Editor/MainFrame/Private/Menus/MainMenu.cpp
--- a/Engine/Source/Editor/MainFrame/Private/Menus/MainMenu.cpp
+++ b/Engine/Source/Editor/MainFrame/Private/Menus/MainMenu.cpp
@@ -12,6 +12,7 @@
 #include "Editor/DeviceProfileEditor/Public/DeviceProfileEditorModule.h"
 #include "UndoHistoryModule.h"
 #include "GenericCommands.h"
+#include "ToolboxModule.h"
 
 #define LOCTEXT_NAMESPACE "MainFileMenu"
 
@@ -274,10 +275,18 @@ void FMainMenu::FillHelpMenu( FMenuBuilder& MenuBuilder, const TSharedRef< FExte
 }
 
 
-TSharedRef<SWidget> FMainMenu::MakeMainMenu( const TSharedPtr<FTabManager>& TabManager, const TSharedRef< FExtender > Extender )
+TSharedRef<SWidget> FMainMenu::MakeMainMenu(const TSharedPtr<FTabManager>& TabManager, const TSharedRef< FExtender > Extender)
 {
 #define LOCTEXT_NAMESPACE "MainMenu"
 
+	
+	// Put the toolbox into our menus
+	{
+		const IWorkspaceMenuStructure& MenuStructure = WorkspaceMenu::GetMenuStructure();
+		IToolboxModule& ToolboxModule = FModuleManager::LoadModuleChecked<IToolboxModule>("Toolbox");
+		ToolboxModule.RegisterSpawners(MenuStructure.GetDeveloperToolsDebugCategory(), MenuStructure.GetDeveloperToolsMiscCategory());
+	}
+
 	// Cache all project names once
 	FMainFrameActionCallbacks::CacheProjectNames();
 
diff --git a/Engine/Source/Runtime/Slate/Private/Framework/Application/SlateApplication.cpp b/Engine/Source/Runtime/Slate/Private/Framework/Application/SlateApplication.cpp
--- a/Engine/Source/Runtime/Slate/Private/Framework/Application/SlateApplication.cpp
+++ b/Engine/Source/Runtime/Slate/Private/Framework/Application/SlateApplication.cpp
@@ -9,6 +9,7 @@
 #include "GenericCommands.h"
 #include "NotificationManager.h"
 #include "IInputProcessor.h"
+#include "ToolboxModule.h"
 
 
 class FEventRouter
@@ -3556,8 +3557,18 @@ bool FSlateApplication::ProcessKeyDownEvent( FKeyEvent& InKeyEvent )
 					Reply = FReply::Handled();
 
 					return Reply.IsEventHandled();
-				}
 			}
+		}
+
+#if !(UE_BUILD_SHIPPING || UE_BUILD_TEST)
+		// Ctrl+Shift+~ summons the Toolbox.
+		if (InKeyEvent.GetKey() == EKeys::Tilde && InKeyEvent.IsControlDown() && InKeyEvent.IsShiftDown())
+		{
+			IToolboxModule& ToolboxModule = FModuleManager::LoadModuleChecked<IToolboxModule>("Toolbox");
+			ToolboxModule.SummonToolbox();
+		}
+
+#endif //!(UE_BUILD_SHIPPING || UE_BUILD_TEST)
 
 		// Bubble the keyboard event
 		FWidgetPath EventPath = UserFocusEntries[InKeyEvent.GetUserIndex()].WidgetPath.ToWidgetPath();
diff --git a/Engine/Source/Runtime/Slate/Slate.Build.cs b/Engine/Source/Runtime/Slate/Slate.Build.cs
--- a/Engine/Source/Runtime/Slate/Slate.Build.cs
+++ b/Engine/Source/Runtime/Slate/Slate.Build.cs
@@ -55,6 +55,8 @@ public class Slate : ModuleRules
 				"Runtime/Slate/Private/Widgets/Text",
 				"Runtime/Slate/Private/Widgets/Tutorials",
 				"Runtime/Slate/Private/Widgets/Views",
+
+                "Runtime/Toolbox/Public/"
 			}
 		);
 
diff --git a/Engine/Source/Runtime/Toolbox/Private/ToolboxModule.cpp b/Engine/Source/Runtime/Toolbox/Private/ToolboxModule.cpp
new file mode 100644
--- /dev/null
+++ b/Engine/Source/Runtime/Toolbox/Private/ToolboxModule.cpp
@@ -0,0 +1,225 @@
+// Copyright 1998-2015 Epic Games, Inc. All Rights Reserved.
+
+#include "ToolboxModule.h"
+#include "SlateBasics.h"
+#include "ModuleManager.h"
+#include "GammaUI.h"
+#include "ModuleUIInterface.h"
+#include "MainFrame.h"
+#include "DesktopPlatformModule.h"
+#include "SDockTab.h"
+#include "SNotificationList.h"
+#include "NotificationManager.h"
+#include "STestSuite.h"
+#include "ISlateReflectorModule.h"
+
+
+static bool bTabsRegistered = false;
+
+static bool CanShowModulesTab()
+{
+	FString SolutionPath;
+	return FDesktopPlatformModule::Get()->GetSolutionPath(SolutionPath);
+}
+
+class SDebugPanel : public SCompoundWidget
+{
+public:
+
+	SLATE_BEGIN_ARGS(SDebugPanel){}
+	SLATE_END_ARGS()
+
+	BEGIN_SLATE_FUNCTION_BUILD_OPTIMIZATION
+	void Construct( const FArguments& )
+	{
+		ChildSlot
+		[
+			SNew( SVerticalBox )
+			+ SVerticalBox::Slot()
+			.AutoHeight()
+			.Padding( 4.0f )
+			.HAlign(HAlign_Left)
+			[
+				SNew( SButton )
+				.Text( NSLOCTEXT("DeveloperToolbox", "ReloadTextures", "Reload Textures") )
+				.OnClicked( this, &SDebugPanel::OnReloadTexturesClicked )
+			]
+			
+			+ SVerticalBox::Slot()
+			.AutoHeight()
+			.Padding( 4.0f )
+			.HAlign(HAlign_Left)
+			[
+				SNew( SButton )
+				.Text( NSLOCTEXT("DeveloperToolbox", "FlushFontCache", "Flush Font Cache") )
+				.OnClicked( this, &SDebugPanel::OnFlushFontCacheClicked )
+			]
+			+ SVerticalBox::Slot()
+			.AutoHeight()
+			.Padding( 4.0f )
+			.HAlign(HAlign_Left)
+			[
+				SNew( SButton )
+				.Text( NSLOCTEXT("DeveloperToolbox", "TestSuite", "Test Suite") )
+				.OnClicked( this, &SDebugPanel::OnTestSuiteClicked )
+			]
+			+ SVerticalBox::Slot()
+			.AutoHeight()
+			.Padding( 4.0f )
+			.HAlign(HAlign_Left)
+			[
+				SNew( SButton )
+				.Text( NSLOCTEXT("DeveloperToolbox", "DisplayTextureAtlases", "Display Texture Atlases") )
+				.OnClicked( this, &SDebugPanel::OnDisplayTextureAtlases )
+			]
+			+ SVerticalBox::Slot()
+			.AutoHeight()
+			.Padding( 4.0f )
+			.HAlign(HAlign_Left)
+			[
+				SNew( SButton )
+				.Text( NSLOCTEXT("DeveloperToolbox", "DisplayFontAtlases", "Display Font Atlases") )
+				.OnClicked( this, &SDebugPanel::OnDisplayFontAtlases )
+			]
+		];
+	}
+	END_SLATE_FUNCTION_BUILD_OPTIMIZATION
+
+	
+	FReply OnReloadTexturesClicked()
+	{
+		FSlateApplication::Get().GetRenderer()->ReloadTextureResources();
+		return FReply::Handled();
+	}
+
+	FReply OnDisplayTextureAtlases()
+	{
+		static const FName SlateReflectorModuleName("SlateReflector");
+		FModuleManager::LoadModuleChecked<ISlateReflectorModule>(SlateReflectorModuleName).DisplayTextureAtlasVisualizer();
+		return FReply::Handled();
+	}
+
+	FReply OnDisplayFontAtlases()
+	{
+		static const FName SlateReflectorModuleName("SlateReflector");
+		FModuleManager::LoadModuleChecked<ISlateReflectorModule>(SlateReflectorModuleName).DisplayFontAtlasVisualizer();
+		return FReply::Handled();
+	}
+
+	FReply OnFlushFontCacheClicked()
+	{
+		FSlateApplication::Get().GetRenderer()->FlushFontCache();
+		return FReply::Handled();
+	}
+
+	FReply OnTestSuiteClicked()
+	{
+		RestoreSlateTestSuite();
+		return FReply::Handled();
+	}
+};
+
+TSharedRef<SDockTab> CreateDebugToolsTab( const FSpawnTabArgs& Args )
+{
+	FGammaUI& GammaUI = FModuleManager::Get().LoadModuleChecked<FGammaUI>( FName("GammaUI") );
+
+	return
+		SNew(SDockTab)
+		.TabRole( ETabRole::NomadTab )
+		[
+			SNew(SVerticalBox)
+			+SVerticalBox::Slot()
+			.AutoHeight()
+			[
+				SNew( SDebugPanel )
+			]
+			+SVerticalBox::Slot()
+			.AutoHeight()
+			[
+				GammaUI.GetGammaUIPanel().ToSharedRef()
+			]
+		];
+}
+
+TSharedRef<SDockTab> CreateModulesTab( const FSpawnTabArgs& Args )
+{
+	IModuleUIInterface& ModuleUI = FModuleManager::Get().LoadModuleChecked<IModuleUIInterface>( FName("ModuleUI") );
+
+	return
+		SNew(SDockTab)
+		.TabRole( ETabRole::NomadTab )
+		[
+			SNew(SBorder)
+			.Padding(2)
+			.BorderImage( FEditorStyle::GetBrush("ToolPanel.GroupBorder") )
+			[
+				ModuleUI.GetModuleUIWidget()
+			]
+		];
+}
+
+
+
+class FToolboxModule : public IToolboxModule
+{
+	void StartupModule()
+	{
+
+	}
+
+	void ShutdownModule()
+	{
+		bTabsRegistered = false;
+
+		if (FSlateApplication::IsInitialized())
+		{
+			FGlobalTabmanager::Get()->UnregisterNomadTabSpawner("DebugTools");
+			if (CanShowModulesTab())
+			{
+				FGlobalTabmanager::Get()->UnregisterNomadTabSpawner("ModulesTab");
+			}
+		}
+	}
+
+	virtual void RegisterSpawners(const TSharedPtr<FWorkspaceItem>& DebugToolsTabCategory, const TSharedPtr<FWorkspaceItem>& ModulesTabCategory) override
+	{
+		if (!bTabsRegistered)
+		{
+			bTabsRegistered = true;
+			{
+				FTabSpawnerEntry& Spawner = FGlobalTabmanager::Get()->RegisterNomadTabSpawner("DebugTools", FOnSpawnTab::CreateStatic(&CreateDebugToolsTab))
+					.SetDisplayName(NSLOCTEXT("Toolbox", "DebugTools", "Debug Tools"))
+					.SetTooltipText(NSLOCTEXT("Toolbox", "DebugToolsTooltipText", "Open the Debug Tools tab."))
+					.SetIcon(FSlateIcon(FEditorStyle::GetStyleSetName(), "DebugTools.TabIcon"));
+
+				if (DebugToolsTabCategory.IsValid())
+				{
+					Spawner.SetGroup(DebugToolsTabCategory.ToSharedRef());
+				}
+			}
+
+			if (CanShowModulesTab())
+			{
+				FTabSpawnerEntry& Spawner = FGlobalTabmanager::Get()->RegisterNomadTabSpawner("ModulesTab", FOnSpawnTab::CreateStatic(&CreateModulesTab))
+					.SetDisplayName(NSLOCTEXT("Toolbox", "Modules", "Modules"))
+					.SetTooltipText(NSLOCTEXT("Toolbox", "ModulesTooltipText", "Open the Modules tab."))
+					.SetIcon(FSlateIcon(FEditorStyle::GetStyleSetName(), "Modules.TabIcon"));
+
+				if (ModulesTabCategory.IsValid())
+				{
+					Spawner.SetGroup(ModulesTabCategory.ToSharedRef());
+				}
+			}
+		}
+	}
+
+	virtual void SummonToolbox() override
+	{
+		RegisterSpawners(nullptr, nullptr);
+		FGlobalTabmanager::Get()->InvokeTab(FTabId("DebugTools"));
+	}
+};
+
+
+IMPLEMENT_MODULE(FToolboxModule, Toolbox);
+
diff --git a/Engine/Source/Runtime/Toolbox/Public/ToolboxModule.h b/Engine/Source/Runtime/Toolbox/Public/ToolboxModule.h
new file mode 100644
--- /dev/null
+++ b/Engine/Source/Runtime/Toolbox/Public/ToolboxModule.h
@@ -0,0 +1,28 @@
+// Copyright 1998-2015 Epic Games, Inc. All Rights Reserved.
+
+#pragma once
+
+#include "ModuleInterface.h"
+
+class FWorkspaceItem;
+
+/**
+ * The toolbox module is a lightweight UI for quickly summoning
+ * development tools: profilers, widget inspector, etc.
+ */
+class IToolboxModule : public IModuleInterface
+{
+public:
+
+	/**
+	 * Register spawners for tabs provided by IToolboxModule. Only the first call to this method has any effect.
+	 * Subsequent calls will be ignored until the module is reloaded.
+	 *
+	 * @param  DebugToolsTabCategory       Optional menu category for the debug tools tab.
+	 * @param  ModulesTabCategory          Optional menu category for the modules tab.
+	 */
+	virtual void RegisterSpawners(const TSharedPtr<FWorkspaceItem>& DebugToolsTabCategory, const TSharedPtr<FWorkspaceItem>& ModulesTabCategory) = 0;
+
+	/** Open the toolbox tab,  */
+	virtual void SummonToolbox() = 0;
+};
diff --git a/Engine/Source/Runtime/Toolbox/Toolbox.Build.cs b/Engine/Source/Runtime/Toolbox/Toolbox.Build.cs
new file mode 100644
--- /dev/null
+++ b/Engine/Source/Runtime/Toolbox/Toolbox.Build.cs
@@ -0,0 +1,38 @@
+// Copyright 1998-2015 Epic Games, Inc. All Rights Reserved.
+
+using UnrealBuildTool;
+
+public class Toolbox : ModuleRules
+{
+	public Toolbox(TargetInfo Target)
+	{
+		PrivateIncludePathModuleNames.AddRange(
+			new string[] {
+				"GammaUI",
+				"MainFrame",
+				"ModuleUI"
+			}
+		);
+
+		PrivateDependencyModuleNames.AddRange(
+			new string[] {
+				"Core",
+                "InputCore",
+				"Slate",
+				"SlateCore",
+				"SlateReflector",
+                "EditorStyle",
+				"DesktopPlatform",
+				"AppFramework"
+			}
+		);
+
+		DynamicallyLoadedModuleNames.AddRange(
+			new string[] {
+				"GammaUI", 
+				"MainFrame",
+				"ModuleUI"
+			}
+		);
+	}
+}
\ No newline at end of file
